LiveCoding: Fix horizontal scroll bar not working in console application.

diff --git a/Engine/Source/Programs/LiveCodingConsole/Private/SLogWidget.cpp b/Engine/Source/Programs/LiveCodingConsole/Private/SLogWidget.cpp
--- a/Engine/Source/Programs/LiveCodingConsole/Private/SLogWidget.cpp
+++ b/Engine/Source/Programs/LiveCodingConsole/Private/SLogWidget.cpp
@@ -75,7 +75,8 @@ int32 FLogWidgetTextLayoutMarshaller::GetNumLines() const
 //// SLogWidget ////
 
 SLogWidget::SLogWidget()
-	: bIsUserScrolled(false)
+	: bIsUserScrolledX(false)
+	, bIsUserScrolledY(false)
 {
 }
 
@@ -88,17 +89,18 @@ void SLogWidget::Construct(const FArguments& InArgs)
 	MessagesTextMarshaller = MakeShared<FLogWidgetTextLayoutMarshaller>();
 
 	ChildSlot
-	[
-		SNew(SBorder)
 		[
-			 SAssignNew(MessagesTextBox, SMultiLineEditableTextBox)
+			SNew(SBorder)
+			[
+				SAssignNew(MessagesTextBox, SMultiLineEditableTextBox)
 				.Style(FLiveCodingConsoleStyle::Get(), "Log.TextBox")
-				.Marshaller(MessagesTextMarshaller)
-				.IsReadOnly(true)
-				.AlwaysShowScrollbars(true)
-				.OnVScrollBarUserScrolled(this, &SLogWidget::OnScroll)
-		]
-	];
+		        .Marshaller(MessagesTextMarshaller)
+		        .IsReadOnly(true)
+		        .AlwaysShowScrollbars(true)
+		        .OnHScrollBarUserScrolled(this, &SLogWidget::OnScrollX)
+		        .OnVScrollBarUserScrolled(this, &SLogWidget::OnScrollY)
+			]
+		];
 
 	RegisterActiveTimer(0.03f, FWidgetActiveTimerDelegate::CreateSP(this, &SLogWidget::OnTimerElapsed));
 }
@@ -111,7 +113,8 @@ void SLogWidget::Clear()
 void SLogWidget::ScrollToEnd()
 {
 	MessagesTextBox->ScrollTo(FTextLocation(MessagesTextMarshaller->GetNumLines() - 1));
-	bIsUserScrolled = false;
+	bIsUserScrolledX = false;
+	bIsUserScrolledY = false;
 }
 
 void SLogWidget::AppendLine(const FSlateColor& Color, const FString& Text)
@@ -124,9 +127,14 @@ void SLogWidget::AppendLine(const FSlateColor& Color, const FString& Text)
 	QueuedLines.Add(MoveTemp(Line));
 }
 
-void SLogWidget::OnScroll(float ScrollOffset)
+void SLogWidget::OnScrollX(float ScrollOffset)
 {
-	bIsUserScrolled = ScrollOffset < 1.0 && !FMath::IsNearlyEqual(ScrollOffset, 1.0f);
+	bIsUserScrolledX = ScrollOffset > 0.0 && !FMath::IsNearlyEqual(ScrollOffset, 0.0f);
+}
+
+void SLogWidget::OnScrollY(float ScrollOffset)
+{
+	bIsUserScrolledY = ScrollOffset < 1.0 && !FMath::IsNearlyEqual(ScrollOffset, 1.0f);
 }
 
 EActiveTimerReturnType SLogWidget::OnTimerElapsed(double CurrentTime, float DeltaTime)
@@ -136,7 +144,7 @@ EActiveTimerReturnType SLogWidget::OnTimerElapsed(double CurrentTime, float Delt
 	{
 		MessagesTextMarshaller->AppendLine(QueuedLine.Color, QueuedLine.Text);
 	}
-	if(!bIsUserScrolled)
+	if(!bIsUserScrolledX && !bIsUserScrolledY)
 	{
 		ScrollToEnd();
 	}
diff --git a/Engine/Source/Programs/LiveCodingConsole/Private/SLogWidget.h b/Engine/Source/Programs/LiveCodingConsole/Private/SLogWidget.h
--- a/Engine/Source/Programs/LiveCodingConsole/Private/SLogWidget.h
+++ b/Engine/Source/Programs/LiveCodingConsole/Private/SLogWidget.h
@@ -18,7 +18,7 @@ class FLogWidgetTextLayoutMarshaller : public FBaseTextLayoutMarshaller
 public:
 	FLogWidgetTextLayoutMarshaller();
 	virtual ~FLogWidgetTextLayoutMarshaller();
-	
+
 	// ITextLayoutMarshaller
 	virtual void SetText(const FString& SourceString, FTextLayout& TargetTextLayout) override;
 	virtual void GetText(FString& TargetString, const FTextLayout& SourceTextLayout) override;
@@ -39,7 +39,7 @@ public:
 	SLATE_BEGIN_ARGS(SLogWidget) { }
 	SLATE_END_ARGS()
 
-	SLogWidget();
+		SLogWidget();
 	~SLogWidget();
 
 	void Construct( const FArguments& InArgs );
@@ -60,8 +60,10 @@ protected:
 	TArray<FLine> QueuedLines;
 	TSharedPtr<FLogWidgetTextLayoutMarshaller> MessagesTextMarshaller;
 	TSharedPtr<SMultiLineEditableTextBox> MessagesTextBox;
-	bool bIsUserScrolled;
+	bool bIsUserScrolledX;
+	bool bIsUserScrolledY;
 
-	void OnScroll(float ScrollOffset);
+	void OnScrollX(float ScrollOffset);
+	void OnScrollY(float ScrollOffset);
 	EActiveTimerReturnType OnTimerElapsed(double CurrentTime, float DeltaTime);
 };
