Live Coding: Enable high DPI mode for LiveCodingConsole.

diff --git a/Engine/Source/Programs/LiveCodingConsole/Private/LiveCodingConsole.cpp b/Engine/Source/Programs/LiveCodingConsole/Private/LiveCodingConsole.cpp
--- a/Engine/Source/Programs/LiveCodingConsole/Private/LiveCodingConsole.cpp
+++ b/Engine/Source/Programs/LiveCodingConsole/Private/LiveCodingConsole.cpp
@@ -382,12 +382,15 @@ bool LiveCodingConsoleMain(const TCHAR* CmdLine)
 	GEngineLoop.PreInit(CmdLine);
 	check(GConfig && GConfig->IsReadyForUse());
 
+	// Initialize high DPI mode
+	FSlateApplication::InitHighDPI(true);
+
 	{
-		// create the platform slate application (what FSlateApplication::Get() returns)
+		// Create the platform slate application (what FSlateApplication::Get() returns)
 		TSharedRef<FSlateApplication> Slate = FSlateApplication::Create(MakeShareable(FPlatformApplicationMisc::CreateApplication()));
 
 		{
-			// initialize renderer
+			// Initialize renderer
 			TSharedRef<FSlateRenderer> SlateRenderer = GetStandardStandaloneRenderer();
 
 			// Try to initialize the renderer. It's possible that we launched when the driver crashed so try a few times before giving up.
@@ -399,7 +402,7 @@ bool LiveCodingConsoleMain(const TCHAR* CmdLine)
 				return false;
 			}
 
-			// set the normal UE4 GIsRequestingExit when outer frame is closed
+			// Set the normal UE4 GIsRequestingExit when outer frame is closed
 			Slate->SetExitRequestedHandler(FSimpleDelegate::CreateStatic(&OnRequestExit));
 
 			// Prepare the custom Slate styles
