LiveCoding: Fix "Show Console" button not making the window visible if the window has been closed.

diff --git a/Engine/Source/Programs/LiveCodingConsole/Private/LiveCodingConsole.cpp b/Engine/Source/Programs/LiveCodingConsole/Private/LiveCodingConsole.cpp
--- a/Engine/Source/Programs/LiveCodingConsole/Private/LiveCodingConsole.cpp
+++ b/Engine/Source/Programs/LiveCodingConsole/Private/LiveCodingConsole.cpp
@@ -285,6 +285,12 @@ private:
 		}
 	}
 
+	void ShowConsole()
+	{
+		SetVisible(true);
+		BringToFront();
+	}
+
 	void OnCompileStartedAsync()
 	{
 		FScopeLock Lock(&CriticalSection);
@@ -301,7 +307,7 @@ private:
 			Info.bFireAndForget = false;
 			Info.FadeOutDuration = 0.0f;
 			Info.ExpireDuration = 0.0f;
-			Info.Hyperlink = FSimpleDelegate::CreateRaw(this, &FLiveCodingConsoleApp::BringToFront);
+			Info.Hyperlink = FSimpleDelegate::CreateRaw(this, &FLiveCodingConsoleApp::ShowConsole);
 			Info.HyperlinkText = LOCTEXT("BuildStatusShowConsole", "Show Console");
 			Info.ButtonDetails.Add(FNotificationButtonInfo(LOCTEXT("BuildStatusCancel", "Cancel"), FText(), FSimpleDelegate::CreateRaw(this, &FLiveCodingConsoleApp::CancelBuild), SNotificationItem::CS_Pending));
 
