LiveCoding: Restore the console if it's minimized when a compile started.

diff --git a/Engine/Source/Programs/LiveCodingConsole/Private/LiveCodingConsole.cpp b/Engine/Source/Programs/LiveCodingConsole/Private/LiveCodingConsole.cpp
--- a/Engine/Source/Programs/LiveCodingConsole/Private/LiveCodingConsole.cpp
+++ b/Engine/Source/Programs/LiveCodingConsole/Private/LiveCodingConsole.cpp
@@ -150,6 +150,10 @@ private:
 	void BringToFront()
 	{
 		HWND WindowHandle = (HWND)Window->GetNativeWindow()->GetOSWindowHandle();
+		if (IsIconic(WindowHandle))
+		{
+			ShowWindow(WindowHandle, SW_RESTORE);
+		}
 		::SetWindowPos(WindowHandle, HWND_TOPMOST, 0, 0, 0, 0, SWP_NOMOVE | SWP_NOSIZE | SWP_NOACTIVATE);
 		::SetWindowPos(WindowHandle, HWND_NOTOPMOST, 0, 0, 0, 0, SWP_NOMOVE | SWP_NOSIZE | SWP_NOACTIVATE);
 	}
