LiveCoding: Wait for other processes to finish when triggering a compile. Allows live coding to work with multiple editors open.

diff --git a/Engine/Source/Programs/LiveCodingConsole/Private/LiveCodingConsole.cpp b/Engine/Source/Programs/LiveCodingConsole/Private/LiveCodingConsole.cpp
--- a/Engine/Source/Programs/LiveCodingConsole/Private/LiveCodingConsole.cpp
+++ b/Engine/Source/Programs/LiveCodingConsole/Private/LiveCodingConsole.cpp
@@ -196,7 +196,7 @@ private:
 		}
 
 		FString ManifestFileName = FPaths::ConvertRelativePathToFull(FPaths::EngineIntermediateDir() / TEXT("LiveCoding.json"));
-		Arguments += FString::Printf(TEXT("-LiveCoding -LiveCodingManifest=\"%s\""), *ManifestFileName);
+		Arguments += FString::Printf(TEXT("-LiveCoding -LiveCodingManifest=\"%s\" -WaitMutex"), *ManifestFileName);
 
 		AppendLogLine(ELiveCodingLogVerbosity::Info, *FString::Printf(TEXT("Running %s %s"), *Executable, *Arguments));
 
