Fix deprecation warnings in LiveCodingConsole.

diff --git a/Engine/Source/Programs/LiveCodingConsole/Private/LiveCodingConsole.cpp b/Engine/Source/Programs/LiveCodingConsole/Private/LiveCodingConsole.cpp
--- a/Engine/Source/Programs/LiveCodingConsole/Private/LiveCodingConsole.cpp
+++ b/Engine/Source/Programs/LiveCodingConsole/Private/LiveCodingConsole.cpp
@@ -17,6 +17,7 @@
 #include "LiveCodingManifest.h"
 #include "SourceCodeAccess/Public/ISourceCodeAccessModule.h"
 #include "Modules/ModuleInterface.h"
+#include "Widgets/Input/SButton.h"
 
 #define LOCTEXT_NAMESPACE "LiveCodingConsole"
 
@@ -454,7 +455,7 @@ bool LiveCodingConsoleMain(const TCHAR* CmdLine)
 			FLiveCodingConsoleStyle::Initialize();
 
 			// Set the icon
-			Slate->SetAppIcon(FLiveCodingConsoleStyle::Get().GetBrush("AppIcon"));
+			FAppStyle::SetAppStyleSet(FLiveCodingConsoleStyle::Get());
 
 			// Load the source code access module
 			ISourceCodeAccessModule& SourceCodeAccessModule = FModuleManager::LoadModuleChecked<ISourceCodeAccessModule>(FName("SourceCodeAccess"));
