Fix live coding console styling

diff --git a/Engine/Source/Programs/LiveCodingConsole/Private/LiveCodingConsoleStyle.cpp b/Engine/Source/Programs/LiveCodingConsole/Private/LiveCodingConsoleStyle.cpp
--- a/Engine/Source/Programs/LiveCodingConsole/Private/LiveCodingConsoleStyle.cpp
+++ b/Engine/Source/Programs/LiveCodingConsole/Private/LiveCodingConsoleStyle.cpp
@@ -47,6 +47,9 @@ TSharedRef< FSlateStyleSet > FLiveCodingConsoleStyle::Create()
 	TSharedRef<FSlateStyleSet> StyleRef = MakeShareable(new FSlateStyleSet("LiveCodingServerStyle"));
 	FSlateStyleSet& Style = StyleRef.Get();
 
+	Style.SetParentStyleName("CoreStyle");
+
+
 	Style.Set( "AppIcon", new IMAGE_BRUSH( "Icons/DefaultAppIcon", FVector2D(20,20)) );
 
 	const FTextBlockStyle DefaultText = FTextBlockStyle()
