Fixed up paths to UBT in LivingCoding to be ini driven, allows us to update the path when we move UBT as part of netcore tranisition.

diff --git a/Engine/Source/Programs/LiveCodingConsole/Private/LiveCodingConsole.cpp b/Engine/Source/Programs/LiveCodingConsole/Private/LiveCodingConsole.cpp
--- a/Engine/Source/Programs/LiveCodingConsole/Private/LiveCodingConsole.cpp
+++ b/Engine/Source/Programs/LiveCodingConsole/Private/LiveCodingConsole.cpp
@@ -210,7 +210,16 @@ private:
 		NextPatchStartTime = FDateTime::UtcNow();
 
 		// Get the UBT path
-		FString Executable = FPaths::EngineDir() / TEXT("Binaries/DotNET/UnrealBuildTool.exe");
+		FString Executable;
+		FString Entry;
+		if( GConfig->GetString( TEXT("PlatformPaths"), TEXT("UnrealBuildTool"), Entry, GEngineIni ))
+		{
+			Executable = FPaths::ConvertRelativePathToFull(FPaths::EngineDir() / Entry);
+		}
+		else
+		{
+			Executable = FPaths::EngineDir() / TEXT("Binaries/DotNET/UnrealBuildTool.exe");
+		}
 		FPaths::MakePlatformFilename(Executable);
 
 		// Write out the list of lazy-loaded modules for UBT to check
