Minor fix to how UBT is located.  The ini file setting includes the Engine directory so the prefix path needed to be changed to root instead of engine.

diff --git a/Engine/Source/Programs/LiveCodingConsole/Private/LiveCodingConsole.cpp b/Engine/Source/Programs/LiveCodingConsole/Private/LiveCodingConsole.cpp
--- a/Engine/Source/Programs/LiveCodingConsole/Private/LiveCodingConsole.cpp
+++ b/Engine/Source/Programs/LiveCodingConsole/Private/LiveCodingConsole.cpp
@@ -214,7 +214,7 @@ private:
 		FString Entry;
 		if( GConfig->GetString( TEXT("PlatformPaths"), TEXT("UnrealBuildTool"), Entry, GEngineIni ))
 		{
-			Executable = FPaths::ConvertRelativePathToFull(FPaths::EngineDir() / Entry);
+			Executable = FPaths::ConvertRelativePathToFull(FPaths::RootDir() / Entry);
 		}
 		else
 		{
