Added a limit to the number of actions live coding will handle by default.  The user can change these limits and even request that the compilation try again withot enforcing the limit.

diff --git a/Engine/Config/BaseEngine.ini b/Engine/Config/BaseEngine.ini
--- a/Engine/Config/BaseEngine.ini
+++ b/Engine/Config/BaseEngine.ini
@@ -2532,4 +2532,8 @@ DefaultChaosSolverActorClass=/Script/ChaosSolverEngine.ChaosSolverActor
 DefaultSizeInMegabyte=64
 
 [PlatformPaths]
-UnrealBuildTool=Engine/Binaries/DotNET/UnrealBuildTool/UnrealBuildTool.exe
\ No newline at end of file
+UnrealBuildTool=Engine/Binaries/DotNET/UnrealBuildTool/UnrealBuildTool.exe
+
+[LiveCoding]
+bDisableActionLimit=false
+ActionLimit=100
\ No newline at end of file
diff --git a/Engine/Source/Developer/Windows/LiveCodingServer/Private/External/LC_ServerCommandThread.cpp b/Engine/Source/Developer/Windows/LiveCodingServer/Private/External/LC_ServerCommandThread.cpp
--- a/Engine/Source/Developer/Windows/LiveCodingServer/Private/External/LC_ServerCommandThread.cpp
+++ b/Engine/Source/Developer/Windows/LiveCodingServer/Private/External/LC_ServerCommandThread.cpp
@@ -726,6 +726,7 @@ void ServerCommandThread::CompileChanges(bool didAllProcessesMakeProgress)
 
 		// Keep retrying the compile until we've added all the required modules
 		TMap<FString, TArray<FString>> ModuleToObjectFiles;
+		ELiveCodingCompileReason CompileReason = ELiveCodingCompileReason::Initial;
 		for (;;)
 		{
 			// Build a list of modules which are enabled for live coding
@@ -737,17 +738,30 @@ void ServerCommandThread::CompileChanges(bool didAllProcessesMakeProgress)
 
 			// Execute the compile
 			TArray<FString> RequiredModules;
-			if (CompileDelegate.Execute(Targets, ValidModules, RequiredModules, ModuleToObjectFiles))
+			ELiveCodingCompileResult CompileResult = CompileDelegate.Execute(Targets, ValidModules, RequiredModules, ModuleToObjectFiles, CompileReason);
+			if (CompileResult == ELiveCodingCompileResult::Success)
 			{
 				break;
 			}
-
+			else if (CompileResult == ELiveCodingCompileResult::Canceled)
+			{
+				GLiveCodingServer->GetCompileFinishedDelegate().ExecuteIfBound(ELiveCodingResult::Error, L"Compilation canceled.");
+				return;
+			}
+			else if (CompileResult == ELiveCodingCompileResult::Failure)
+			{
+				GLiveCodingServer->GetCompileFinishedDelegate().ExecuteIfBound(ELiveCodingResult::Error, L"Compilation error.");
+				return;
+			}
+			
 			// Enable any lazy-loaded modules that we need
-			if (!EnableRequiredModules(RequiredModules))
+			if (!RequiredModules.IsEmpty() && !EnableRequiredModules(RequiredModules))
 			{
 				GLiveCodingServer->GetCompileFinishedDelegate().ExecuteIfBound(ELiveCodingResult::Error, L"Compilation error.");
 				return;
 			}
+
+			CompileReason = ELiveCodingCompileReason::Retry;
 		}
 
 		// Reset the unity file cache
diff --git a/Engine/Source/Developer/Windows/LiveCodingServer/Public/ILiveCodingServer.h b/Engine/Source/Developer/Windows/LiveCodingServer/Public/ILiveCodingServer.h
--- a/Engine/Source/Developer/Windows/LiveCodingServer/Public/ILiveCodingServer.h
+++ b/Engine/Source/Developer/Windows/LiveCodingServer/Public/ILiveCodingServer.h
@@ -24,6 +24,20 @@ enum class ELiveCodingLogVerbosity
 	Failure,
 };
 
+enum class ELiveCodingCompileReason
+{
+	Initial,
+	Retry,
+};
+
+enum class ELiveCodingCompileResult
+{
+	Success,
+	Canceled,
+	Failure,
+	Retry,
+};
+
 class ILiveCodingServerModule : public IModuleInterface
 {
 public:
@@ -55,7 +69,7 @@ public:
 	virtual FLogOutputDelegate& GetLogOutputDelegate() = 0;
 
 	typedef TMap<FString, TArray<FString>> FModuleToObjectFiles;
-	DECLARE_DELEGATE_RetVal_FourParams(bool, FCompileDelegate, const TArray<FString>&, const TArray<FString>&, TArray<FString>&, FModuleToObjectFiles&)
+	DECLARE_DELEGATE_RetVal_FiveParams(ELiveCodingCompileResult, FCompileDelegate, const TArray<FString>&, const TArray<FString>&, TArray<FString>&, FModuleToObjectFiles&, ELiveCodingCompileReason);
 	virtual FCompileDelegate& GetCompileDelegate() = 0;
 
 	DECLARE_DELEGATE(FCompileStartedDelegate);
diff --git a/Engine/Source/Programs/LiveCodingConsole/Private/LiveCodingConsole.cpp b/Engine/Source/Programs/LiveCodingConsole/Private/LiveCodingConsole.cpp
--- a/Engine/Source/Programs/LiveCodingConsole/Private/LiveCodingConsole.cpp
+++ b/Engine/Source/Programs/LiveCodingConsole/Private/LiveCodingConsole.cpp
@@ -18,6 +18,10 @@
 #include "SourceCodeAccess/Public/ISourceCodeAccessModule.h"
 #include "Modules/ModuleInterface.h"
 #include "Widgets/Input/SButton.h"
+#include "Widgets/Input/SCheckBox.h"
+#include "Widgets/Text/STextBlock.h"
+#include "Misc/CompilationResult.h"
+#include "Misc/MessageDialog.h"
 
 #define LOCTEXT_NAMESPACE "LiveCodingConsole"
 
@@ -39,6 +43,7 @@ private:
 	TSharedPtr<SNotificationItem> CompileNotification;
 	TArray<FSimpleDelegate> MainThreadTasks;
 	bool bRequestCancel;
+	bool bDisableLimit;
 	FDateTime LastPatchTime;
 	FDateTime NextPatchStartTime;
 
@@ -47,6 +52,7 @@ public:
 		: Slate(InSlate)
 		, Server(InServer)
 		, bRequestCancel(false)
+		, bDisableLimit(false)
 		, LastPatchTime(FDateTime::MinValue())
 		, NextPatchStartTime(FDateTime::MinValue())
 	{
@@ -81,9 +87,26 @@ public:
 					.AutoHeight()
 					.Padding(0.0f, 6.0f, 0.0f, 4.0f)
 					[
-						SNew(SButton)
-						.Text(LOCTEXT("QuickRestart", "Quick Restart"))
-						.OnClicked(FOnClicked::CreateRaw(this, &FLiveCodingConsoleApp::RestartTargets))
+						SNew(SHorizontalBox)
+						+ SHorizontalBox::Slot()
+						.HAlign(HAlign_Center)
+						[
+							SNew(SButton)
+							.Text(LOCTEXT("QuickRestart", "Quick Restart"))
+							.OnClicked(FOnClicked::CreateRaw(this, &FLiveCodingConsoleApp::RestartTargets))
+						]
+						+ SHorizontalBox::Slot()
+						.HAlign(HAlign_Center)
+						.Padding(5)
+						[
+							SNew(SCheckBox)
+							.IsChecked_Lambda([this]() { return bDisableLimit ? ECheckBoxState::Checked : ECheckBoxState::Unchecked; })
+							.OnCheckStateChanged_Lambda([this](ECheckBoxState InCheckBoxState) { bDisableLimit = InCheckBoxState == ECheckBoxState::Checked; })
+							[
+								SNew(STextBlock)
+								.Text(LOCTEXT("DisableLimit", "Disable action limit for this session"))
+							]
+						]
 					]
 				]
 			];
@@ -204,7 +227,7 @@ private:
 		LogWidget->AppendLine(GetLogColor(Verbosity), MoveTemp(Text));
 	}
 
-	bool CompilePatch(const TArray<FString>& Targets, const TArray<FString>& ValidModules, TArray<FString>& RequiredModules, TMap<FString, TArray<FString>>& ModuleToObjectFiles)
+	ELiveCodingCompileResult CompilePatch(const TArray<FString>& Targets, const TArray<FString>& ValidModules, TArray<FString>& RequiredModules, TMap<FString, TArray<FString>>& ModuleToObjectFiles, ELiveCodingCompileReason CompileReason)
 	{
 		// Update the compile start time. This gets copied into the last patch time once a patch has been confirmed to have been applied.
 		NextPatchStartTime = FDateTime::UtcNow();
@@ -241,6 +264,17 @@ private:
 			Arguments += FString::Printf(TEXT("-Target=\"%s\" "), *Target.Replace(TEXT("\""), TEXT("\"\"")));
 		}
 		Arguments += FString::Printf(TEXT("-LiveCoding -LiveCodingModules=\"%s\" -LiveCodingManifest=\"%s\" -WaitMutex"), *ModulesFileName, *ManifestFileName);
+		if (!bDisableLimit && CompileReason == ELiveCodingCompileReason::Initial)
+		{
+			bool bDisableActionLimit = false;
+			GConfig->GetBool(TEXT("LiveCoding"), TEXT("bDisableActionLimit"), bDisableActionLimit, GEngineIni);
+			int ActionLimit = 1;
+			GConfig->GetInt(TEXT("LiveCoding"), TEXT("ActionLimit"), ActionLimit, GEngineIni);
+			if (!bDisableActionLimit && ActionLimit > 0)
+			{
+				Arguments += FString::Printf(TEXT(" -LiveCodingLimit=%d"), ActionLimit);
+			}
+		}
 		AppendLogLine(ELiveCodingLogVerbosity::Info, *FString::Printf(TEXT("Running %s %s"), *Executable, *Arguments));
 
 		// Spawn UBT and wait for it to complete (or the compile button to be pressed)
@@ -252,7 +286,7 @@ private:
 			if (HasCancelledBuild())
 			{
 				AppendLogLine(ELiveCodingLogVerbosity::Warning, TEXT("Build cancelled."));
-				return false;
+				return ELiveCodingCompileResult::Canceled;
 			}
 			FPlatformProcess::Sleep(0.1f);
 		}
@@ -260,15 +294,33 @@ private:
 		int ReturnCode = Process.GetReturnCode();
 		if (ReturnCode != 0)
 		{
+			ELiveCodingCompileResult CompileResult = ELiveCodingCompileResult::Failure;
+
+			// If there are missing modules, then we always retry them
 			if (FPaths::FileExists(ModulesOutputFileName))
 			{
 				FFileHelper::LoadFileToStringArray(RequiredModules, *ModulesOutputFileName);
+				if (!RequiredModules.IsEmpty())
+				{
+					CompileResult = ELiveCodingCompileResult::Retry;
+				}
 			}
-			else
+
+			// If we reached the live coding limit, the prompt the user to retry
+			if (ReturnCode == ECompilationResult::LiveCodingLimitError)
+			{
+				const FText Message = LOCTEXT("LimitText", "Live Coding action limit reached.  Do you wish to compile anyway?\n\n"
+				"The limit can be permanently changed or disabled by setting the ActionLimit or DisableActionLimit setting in the engine.");
+				const FText Title = LOCTEXT("LimitTitle", "Live Coding Action Limit Reached");
+				EAppReturnType::Type ReturnType = FMessageDialog::Open(EAppMsgType::YesNo, Message, &Title);
+				CompileResult = ReturnType == EAppReturnType::Yes ? ELiveCodingCompileResult::Retry : ELiveCodingCompileResult::Canceled;
+			}
+
+			if (CompileResult == ELiveCodingCompileResult::Failure)
 			{
 				AppendLogLine(ELiveCodingLogVerbosity::Failure, TEXT("Build failed."));
 			}
-			return false;
+			return CompileResult;
 		}
 
 		// Read the output manifest
@@ -277,7 +329,7 @@ private:
 		if (!Manifest.Read(*ManifestFileName, ManifestFailReason))
 		{
 			AppendLogLine(ELiveCodingLogVerbosity::Failure, *ManifestFailReason);
-			return false;
+			return ELiveCodingCompileResult::Failure;
 		}
 
 		// Override the linker path
@@ -301,7 +353,7 @@ private:
 				}
 			}
 		}
-		return true;
+		return ELiveCodingCompileResult::Success;
 	}
 
 	void CancelBuild()
diff --git a/Engine/Source/Programs/UnrealBuildTool/Configuration/TargetDescriptor.cs b/Engine/Source/Programs/UnrealBuildTool/Configuration/TargetDescriptor.cs
--- a/Engine/Source/Programs/UnrealBuildTool/Configuration/TargetDescriptor.cs
+++ b/Engine/Source/Programs/UnrealBuildTool/Configuration/TargetDescriptor.cs
@@ -64,6 +64,12 @@ namespace UnrealBuildTool
 		[CommandLine("-LiveCodingManifest=")]
 		public FileReference LiveCodingManifest = null;
 
+		/// <summary>
+		/// If a non-zero value, a live coding request will be terminated if more than the given number of actions are required.
+		/// </summary>
+		[CommandLine("-LiveCodingLimit=")]
+		public uint LiveCodingLimit = 0;
+
 		/// <summary>
 		/// Suppress messages about building this target
 		/// </summary>
diff --git a/Engine/Source/Programs/UnrealBuildTool/System/ExternalExecution.cs b/Engine/Source/Programs/UnrealBuildTool/System/ExternalExecution.cs
--- a/Engine/Source/Programs/UnrealBuildTool/System/ExternalExecution.cs
+++ b/Engine/Source/Programs/UnrealBuildTool/System/ExternalExecution.cs
@@ -55,6 +55,11 @@ namespace UnrealBuildTool
 		/// </summary>
 		OtherCompilationError = 6,
 
+		/// <summary>
+		/// Compilation failed due to live coding action limit being exceeded.
+		/// </summary>
+		LiveCodingLimitError = 7,
+
 		/// <summary>
 		/// Compilation is not supported in the current build
 		/// </summary>
diff --git a/Engine/Source/Programs/UnrealBuildTool/System/HotReload.cs b/Engine/Source/Programs/UnrealBuildTool/System/HotReload.cs
--- a/Engine/Source/Programs/UnrealBuildTool/System/HotReload.cs
+++ b/Engine/Source/Programs/UnrealBuildTool/System/HotReload.cs
@@ -484,6 +484,8 @@ namespace UnrealBuildTool
 
 			if (TargetDescriptor.HotReloadMode == HotReloadMode.LiveCoding)
 			{
+				CompilationResult Result = CompilationResult.Succeeded;
+
 				// Make sure we're not overwriting any lazy-loaded modules
 				if (TargetDescriptor.LiveCodingModules != null)
 				{
@@ -520,7 +522,9 @@ namespace UnrealBuildTool
 						{
 							Log.TraceInformation("Module {0} is not currently enabled for Live Coding", ProtectedOutputFile);
 						}
-						throw new CompilationResultException(CompilationResult.Canceled);
+
+						// Note the issue but continue processing to allow the limit to generate an error if hit.
+						Result = CompilationResult.Canceled;
 					}
 				}
 
@@ -535,6 +539,21 @@ namespace UnrealBuildTool
 				// Get a new list of actions to execute now that the graph has been modified
 				TargetActionsToExecute = ActionGraph.GetActionsToExecute(PrerequisiteActions, CppDependencies, History, BuildConfiguration.bIgnoreOutdatedImportLibraries);
 
+				// Check to see if we exceed the limit for live coding actions
+				if (TargetDescriptor.LiveCodingLimit > 0 && TargetDescriptor.LiveCodingLimit < TargetActionsToExecute.Count)
+				{
+					Log.TraceInformation("The live coding request of {0} actions exceeds the number of allowed actions of {1}", TargetActionsToExecute.Count, TargetDescriptor.LiveCodingLimit);
+					Log.TraceInformation("This limit helps to prevent the situation where seemingly simple changes result in large scale rebuilds.");
+					Log.TraceInformation("It can also help to detect when the engine needs to be rebuilt outside of Live Coding due to compiler changes.");
+					Result = CompilationResult.LiveCodingLimitError;
+				}
+
+				// Throw an exception if there is an issue
+				if (Result != CompilationResult.Succeeded)
+				{
+					throw new CompilationResultException(Result);
+				}
+
 				// Output the Live Coding manifest
 				if (TargetDescriptor.LiveCodingManifest != null)
 				{
diff --git a/Engine/Source/Runtime/Core/Public/Misc/CompilationResult.h b/Engine/Source/Runtime/Core/Public/Misc/CompilationResult.h
--- a/Engine/Source/Runtime/Core/Public/Misc/CompilationResult.h
+++ b/Engine/Source/Runtime/Core/Public/Misc/CompilationResult.h
@@ -29,6 +29,8 @@ namespace ECompilationResult
 		FailedDueToEngineChange = 5,
 		/** Compilation failed due to compilation errors */
 		OtherCompilationError = 6,
+		/** Compilation failed due to live coding limit reached */
+		LiveCodingLimitError = 7,
 		/** Compilation is not supported in the current build */
 		Unsupported,
 		/** Unknown error */
@@ -54,6 +56,8 @@ namespace ECompilationResult
 			return TEXT("OtherCompilationError");
 		case ECompilationResult::CrashOrAssert:
 			return TEXT("CrashOrAssert");
+		case ECompilationResult::LiveCodingLimitError:
+			return TEXT("LiveCodingLimitError");
 		case ECompilationResult::Unsupported:
 			return TEXT("Unsupported");
 		};
