Refer to the new UnrealBuilTool.exe location in the fallback path.

#rnx

diff --git a/Engine/Source/Programs/LiveCodingConsole/Private/LiveCodingConsole.cpp b/Engine/Source/Programs/LiveCodingConsole/Private/LiveCodingConsole.cpp
--- a/Engine/Source/Programs/LiveCodingConsole/Private/LiveCodingConsole.cpp
+++ b/Engine/Source/Programs/LiveCodingConsole/Private/LiveCodingConsole.cpp
@@ -242,7 +242,7 @@ private:
 		}
 		else
 		{
-			Executable = FPaths::EngineDir() / TEXT("Binaries/DotNET/UnrealBuildTool.exe");
+			Executable = FPaths::EngineDir() / TEXT("Binaries/DotNET/UnrealBuildTool/UnrealBuildTool.exe");
 		}
 		FPaths::MakePlatformFilename(Executable);
 
