Adding the ability to clear the live coding console log.

#rnx

diff --git a/Engine/Source/Programs/LiveCodingConsole/Private/SLogWidget.cpp b/Engine/Source/Programs/LiveCodingConsole/Private/SLogWidget.cpp
--- a/Engine/Source/Programs/LiveCodingConsole/Private/SLogWidget.cpp
+++ b/Engine/Source/Programs/LiveCodingConsole/Private/SLogWidget.cpp
@@ -103,12 +103,44 @@ void SLogWidget::Construct(const FArguments& InArgs)
 				.SelectWordOnMouseDoubleClick(false)
 		        .OnHScrollBarUserScrolled(this, &SLogWidget::OnScrollX)
 		        .OnVScrollBarUserScrolled(this, &SLogWidget::OnScrollY)
+				.ContextMenuExtender(this, &SLogWidget::ExtendTextBoxMenu)
 			]
 		];
 
 	RegisterActiveTimer(0.03f, FWidgetActiveTimerDelegate::CreateSP(this, &SLogWidget::OnTimerElapsed));
 }
 
+void SLogWidget::ExtendTextBoxMenu(FMenuBuilder& Builder)
+{
+	FUIAction ClearOutputLogAction(
+		FExecuteAction::CreateRaw(this, &SLogWidget::OnClearLog),
+		FCanExecuteAction::CreateSP(this, &SLogWidget::CanClearLog)
+	);
+
+	Builder.AddMenuEntry(
+		NSLOCTEXT("OutputLog", "ClearLogLabel", "Clear Log"),
+		NSLOCTEXT("OutputLog", "ClearLogTooltip", "Clears all log messages"),
+		FSlateIcon(),
+		ClearOutputLogAction
+	);
+}
+
+void SLogWidget::OnClearLog()
+{
+	// Make sure the cursor is back at the start of the log before we clear it
+	MessagesTextBox->GoTo(FTextLocation(0));
+
+	Clear();
+	MessagesTextBox->Refresh();
+	bIsUserScrolledX = false;
+	bIsUserScrolledY = false;
+}
+
+bool SLogWidget::CanClearLog() const
+{
+	return MessagesTextMarshaller->GetNumLines() > 0;
+}
+
 void SLogWidget::Clear()
 {
 	MessagesTextMarshaller->Clear();
diff --git a/Engine/Source/Programs/LiveCodingConsole/Private/SLogWidget.h b/Engine/Source/Programs/LiveCodingConsole/Private/SLogWidget.h
--- a/Engine/Source/Programs/LiveCodingConsole/Private/SLogWidget.h
+++ b/Engine/Source/Programs/LiveCodingConsole/Private/SLogWidget.h
@@ -66,4 +66,19 @@ protected:
 	void OnScrollX(float ScrollOffset);
 	void OnScrollY(float ScrollOffset);
 	EActiveTimerReturnType OnTimerElapsed(double CurrentTime, float DeltaTime);
+
+	/**
+	 * Extends the context menu used by the text box
+	 */
+	void ExtendTextBoxMenu(FMenuBuilder& Builder);
+
+	/**
+	* Called when delete all is selected
+	*/
+	void OnClearLog();
+
+	/**
+	 * Called to determine whether delete all is currently a valid command
+	 */
+	bool CanClearLog() const;
 };
