Fixed issue with LiveCodingConsole where copying text from the log window would contain extra blank lines.

#rnx

diff --git a/Engine/Source/Programs/LiveCodingConsole/Private/SLogWidget.cpp b/Engine/Source/Programs/LiveCodingConsole/Private/SLogWidget.cpp
--- a/Engine/Source/Programs/LiveCodingConsole/Private/SLogWidget.cpp
+++ b/Engine/Source/Programs/LiveCodingConsole/Private/SLogWidget.cpp
@@ -155,12 +155,35 @@ void SLogWidget::ScrollToEnd()
 
 void SLogWidget::AppendLine(const FSlateColor& Color, const FString& Text)
 {
-	FLine Line;
-	Line.Color = Color;
-	Line.Text = Text;
-
-	FScopeLock Lock(&CriticalSection);
-	QueuedLines.Add(MoveTemp(Line));
+	// Split any multi line text block into multiple lines while removing the \r\n or \n from the string.
+	// The ParseIntoArray method will either leave the trailing blank line or eliminate all contained blank 
+	// lines so we use a hand written algorithm.
+	for (const TCHAR* Cur = *Text, *End = Cur + Text.Len(); Cur != End;)
+	{
+		const TCHAR* LineStart = Cur;
+		const TCHAR* LineEnd = Cur;
+		for (; Cur != End; ++Cur)
+		{
+			if (*Cur == '\n')
+			{
+				LineEnd = Cur;
+				++Cur;
+				break;
+			}
+			else if (*Cur == '\r')
+			{
+				LineEnd = Cur;
+				++Cur;
+				if (Cur != End && *Cur == '\n')
+				{
+					++Cur;
+				}
+				break;
+			}
+		}
+		FScopeLock Lock(&CriticalSection);
+		QueuedLines.Add(FLine{ Color, FString(LineEnd - LineStart, LineStart) });
+	}
 }
 
 void SLogWidget::OnScrollX(float ScrollOffset)
