Fixed issue where some log lines were not appearing in the window

#rnx

diff --git a/Engine/Source/Programs/LiveCodingConsole/Private/SLogWidget.cpp b/Engine/Source/Programs/LiveCodingConsole/Private/SLogWidget.cpp
--- a/Engine/Source/Programs/LiveCodingConsole/Private/SLogWidget.cpp
+++ b/Engine/Source/Programs/LiveCodingConsole/Private/SLogWidget.cpp
@@ -161,7 +161,7 @@ void SLogWidget::AppendLine(const FSlateColor& Color, const FString& Text)
 	for (const TCHAR* Cur = *Text, *End = Cur + Text.Len(); Cur != End;)
 	{
 		const TCHAR* LineStart = Cur;
-		const TCHAR* LineEnd = Cur;
+		const TCHAR* LineEnd = End;
 		for (; Cur != End; ++Cur)
 		{
 			if (*Cur == '\n')
