Added log message for when quick start gets re-enabled.

#rnx

diff --git a/Engine/Source/Developer/Windows/LiveCodingServer/Private/External/LC_ServerCommandThread.cpp b/Engine/Source/Developer/Windows/LiveCodingServer/Private/External/LC_ServerCommandThread.cpp
--- a/Engine/Source/Developer/Windows/LiveCodingServer/Private/External/LC_ServerCommandThread.cpp
+++ b/Engine/Source/Developer/Windows/LiveCodingServer/Private/External/LC_ServerCommandThread.cpp
@@ -1802,11 +1802,6 @@ bool ServerCommandThread::actions::SetReinstancingFlow::Execute(const CommandTyp
 {
 	ServerCommandThread* commandThread = static_cast<ServerCommandThread*>(context);
 
-	if (!commandThread->HasReinstancingProcess() && command->enable)
-	{
-		LC_WARNING_USER("Quick restart disabled when re-instancing is enabled.");
-	}
-
 	// protect against accepting this command while compilation is already in progress
 	CriticalSection::ScopedLock lock(&commandThread->m_actionCS);
 
diff --git a/Engine/Source/Programs/LiveCodingConsole/Private/LiveCodingConsole.cpp b/Engine/Source/Programs/LiveCodingConsole/Private/LiveCodingConsole.cpp
--- a/Engine/Source/Programs/LiveCodingConsole/Private/LiveCodingConsole.cpp
+++ b/Engine/Source/Programs/LiveCodingConsole/Private/LiveCodingConsole.cpp
@@ -45,6 +45,7 @@ private:
 	TArray<FSimpleDelegate> MainThreadTasks;
 	bool bRequestCancel;
 	bool bDisableLimit;
+	bool bHasReinstancingProcess;
 	FDateTime LastPatchTime;
 	FDateTime NextPatchStartTime;
 
@@ -54,6 +55,7 @@ public:
 		, Server(InServer)
 		, bRequestCancel(false)
 		, bDisableLimit(false)
+		, bHasReinstancingProcess(false)
 		, LastPatchTime(FDateTime::MinValue())
 		, NextPatchStartTime(FDateTime::MinValue())
 	{
@@ -98,7 +100,23 @@ public:
 							.ToolTipText_Lambda([this]() { return Server.HasReinstancingProcess() ? 
 								LOCTEXT("DisableQuickRestart", "Quick restarting isn't supported when re-instancing is enabled") :
 								LOCTEXT("EnableQuickRestart", "Restart all live coding applications"); })
-							.IsEnabled_Lambda([this]() { return !Server.HasReinstancingProcess(); })
+							.IsEnabled_Lambda([this]() 
+								{ 
+									bool bNewState = Server.HasReinstancingProcess();
+									if (bNewState != bHasReinstancingProcess)
+									{
+										bHasReinstancingProcess = bNewState;
+										if (bHasReinstancingProcess)
+										{
+											LogWidget->AppendLine(GetLogColor(ELiveCodingLogVerbosity::Warning), TEXT("Quick restart disabled when re-instancing is enabled."));
+										}
+										else
+										{
+											LogWidget->AppendLine(GetLogColor(ELiveCodingLogVerbosity::Success), TEXT("Quick restart enabled."));
+										}
+									}
+									return !bHasReinstancingProcess;
+								})
 						]
 						+ SHorizontalBox::Slot()
 						.HAlign(HAlign_Center)
