LiveCodingConsole will now warn user when "Quick Restart" is used after a patch has been created with re-instancing enabled followed by re-instancing being disabled to make the button available again.  User can choose to continue anyway.

#rnx

diff --git a/Engine/Source/Programs/LiveCodingConsole/Private/LiveCodingConsole.cpp b/Engine/Source/Programs/LiveCodingConsole/Private/LiveCodingConsole.cpp
--- a/Engine/Source/Programs/LiveCodingConsole/Private/LiveCodingConsole.cpp
+++ b/Engine/Source/Programs/LiveCodingConsole/Private/LiveCodingConsole.cpp
@@ -46,6 +46,7 @@ private:
 	bool bRequestCancel;
 	bool bDisableLimit;
 	bool bHasReinstancingProcess;
+	bool bWarnOnRestart;
 	FDateTime LastPatchTime;
 	FDateTime NextPatchStartTime;
 
@@ -56,6 +57,7 @@ public:
 		, bRequestCancel(false)
 		, bDisableLimit(false)
 		, bHasReinstancingProcess(false)
+		, bWarnOnRestart(false)
 		, LastPatchTime(FDateTime::MinValue())
 		, NextPatchStartTime(FDateTime::MinValue())
 	{
@@ -393,6 +395,16 @@ private:
 
 	FReply RestartTargets()
 	{
+		if (bWarnOnRestart)
+		{
+			const FText Message = LOCTEXT("RestartWarningText", "Restarting after patching while re-instancing was enabled can lead to unexpected results.\r\n\r\nDo you wish to continue?");
+			const FText Title = LOCTEXT("RestartWarningTitle", "Restarting after patching while re-instancing was enabled?");
+			EAppReturnType::Type ReturnType = FMessageDialog::Open(EAppMsgType::YesNo, Message, &Title);
+			if (ReturnType != EAppReturnType::Yes)
+			{
+				return FReply::Handled();
+			}
+		}
 		Server.RestartTargets();
 		return FReply::Handled();
 	}
@@ -509,6 +521,11 @@ private:
 		}
 		CompileNotification->ExpireAndFadeout();
 		CompileNotification.Reset();
+
+		if (Result == ELiveCodingResult::Success && bHasReinstancingProcess)
+		{
+			bWarnOnRestart = true;
+		}
 	}
 
 	void OnStatusChangedAsync(const FString& Status)
