Change to Live Coding Log TextBox to match UE5 Styling. Remove usage of BlackGroupBorder.

#UE-121406

diff --git a/Engine/Source/Programs/LiveCodingConsole/Private/LiveCodingConsoleStyle.cpp b/Engine/Source/Programs/LiveCodingConsole/Private/LiveCodingConsoleStyle.cpp
--- a/Engine/Source/Programs/LiveCodingConsole/Private/LiveCodingConsoleStyle.cpp
+++ b/Engine/Source/Programs/LiveCodingConsole/Private/LiveCodingConsoleStyle.cpp
@@ -87,10 +87,10 @@ TSharedRef< FSlateStyleSet > FLiveCodingConsoleStyle::Create()
 		.SetHoveredThumbImage(BOX_BRUSH("Common/Scrollbar_Thumb", FMargin(4.f / 16.f)));
 
 	Style.Set("Log.TextBox", FEditableTextBoxStyle(FCoreStyle::Get().GetWidgetStyle<FEditableTextBoxStyle>("NormalEditableTextBox"))
-		.SetBackgroundImageNormal( BOX_BRUSH( "Common/WhiteGroupBorder", FMargin(4.0f/16.0f) ) )
-		.SetBackgroundImageHovered( BOX_BRUSH( "Common/WhiteGroupBorder", FMargin(4.0f/16.0f) ) )
-		.SetBackgroundImageFocused( BOX_BRUSH( "Common/BlackGroupBorder", FMargin(4.0f/16.0f) ) )
-		.SetBackgroundImageReadOnly( BOX_BRUSH( "Common/BlackGroupBorder", FMargin(4.0f/16.0f) ) )
+		.SetBackgroundImageNormal( BOX_BRUSH("Common/WhiteGroupBorder", FMargin(4.0f/16.0f)))
+		.SetBackgroundImageHovered( BOX_BRUSH("Common/WhiteGroupBorder", FMargin(4.0f/16.0f)))
+		.SetBackgroundImageFocused(BOX_BRUSH("Common/WhiteGroupBorder", FMargin(4.0f / 16.0f)))
+		.SetBackgroundImageReadOnly(BOX_BRUSH("Common/WhiteGroupBorder", FMargin(4.0f / 16.0f)))
 		.SetBackgroundColor( FLinearColor(0.015f, 0.015f, 0.015f) )
 		.SetScrollBarStyle(ScrollBar)
 		);
