Make Live Coding enabled by default.

1) New projects will have live coding enabled by default.
2) Re-instancing will be enabled by default
3) The initial start mode of the console will be hidden
4) Moved some console specific configurations to the ini file specific to LCC
5) Added the saving of the disable action limit to the ini file.

#rnx

diff --git a/Engine/Config/BaseEditorPerProjectUserSettings.ini b/Engine/Config/BaseEditorPerProjectUserSettings.ini
--- a/Engine/Config/BaseEditorPerProjectUserSettings.ini
+++ b/Engine/Config/BaseEditorPerProjectUserSettings.ini
@@ -921,4 +921,7 @@ RemoteHoloLensIP=
 MaxBitrate=4000
 
 [/Script/LiveCoding.LiveCodingSettings]
-bAutomaticallyCompileNewClasses=True
\ No newline at end of file
+Startup=AutomaticButHidden
+bEnabled=True
+bEnableReinstancing=True
+bAutomaticallyCompileNewClasses=True
diff --git a/Engine/Config/BaseEngine.ini b/Engine/Config/BaseEngine.ini
--- a/Engine/Config/BaseEngine.ini
+++ b/Engine/Config/BaseEngine.ini
@@ -2839,10 +2839,6 @@ DefaultChaosSolverActorClass=/Script/ChaosSolverEngine.ChaosSolverActor
 [PlatformPaths]
 UnrealBuildTool=Engine/Binaries/DotNET/UnrealBuildTool/UnrealBuildTool.exe
 
-[LiveCoding]
-bDisableActionLimit=false
-ActionLimit=100
-
 [/Script/InterchangeEngine.InterchangeProjectSettings]
 PipelineStacks=(("Generic", (Pipelines=("/Script/InterchangePipelines.InterchangeGenericAssetsPipeline"))))
 DefaultPipelineStack=Generic
diff --git a/Engine/Programs/LiveCodingConsole/Config/DefaultEngine.ini b/Engine/Programs/LiveCodingConsole/Config/DefaultEngine.ini
new file mode 100644
--- /dev/null
+++ b/Engine/Programs/LiveCodingConsole/Config/DefaultEngine.ini
@@ -0,0 +1,3 @@
+[LiveCodingConsole]
+bDisableActionLimit=false
+ActionLimit=100
diff --git a/Engine/Source/Programs/LiveCodingConsole/Private/LiveCodingConsole.cpp b/Engine/Source/Programs/LiveCodingConsole/Private/LiveCodingConsole.cpp
--- a/Engine/Source/Programs/LiveCodingConsole/Private/LiveCodingConsole.cpp
+++ b/Engine/Source/Programs/LiveCodingConsole/Private/LiveCodingConsole.cpp
@@ -36,6 +36,10 @@ static void OnRequestExit()
 class FLiveCodingConsoleApp
 {
 private:
+	static constexpr TCHAR* SectionName = TEXT("LiveCodingConsole");
+	static constexpr TCHAR* DisableActionLimitKey = TEXT("bDisableActionLimit");
+	static constexpr TCHAR* ActionLimitKey = TEXT("ActionLimit");
+
 	FCriticalSection CriticalSection;
 	FSlateApplication& Slate;
 	ILiveCodingServer& Server;
@@ -44,7 +48,7 @@ private:
 	TSharedPtr<SNotificationItem> CompileNotification;
 	TArray<FSimpleDelegate> MainThreadTasks;
 	bool bRequestCancel;
-	bool bDisableLimit;
+	bool bDisableActionLimit;
 	bool bHasReinstancingProcess;
 	bool bWarnOnRestart;
 	FDateTime LastPatchTime;
@@ -55,7 +59,7 @@ public:
 		: Slate(InSlate)
 		, Server(InServer)
 		, bRequestCancel(false)
-		, bDisableLimit(false)
+		, bDisableActionLimit(false)
 		, bHasReinstancingProcess(false)
 		, bWarnOnRestart(false)
 		, LastPatchTime(FDateTime::MinValue())
@@ -65,6 +69,8 @@ public:
 
 	void Run()
 	{
+		GConfig->GetBool(SectionName, DisableActionLimitKey, bDisableActionLimit, GEngineIni);
+
 		// open up the app window	
 		LogWidget = SNew(SLogWidget);
 
@@ -125,8 +131,13 @@ public:
 						.Padding(5)
 						[
 							SNew(SCheckBox)
-							.IsChecked_Lambda([this]() { return bDisableLimit ? ECheckBoxState::Checked : ECheckBoxState::Unchecked; })
-							.OnCheckStateChanged_Lambda([this](ECheckBoxState InCheckBoxState) { bDisableLimit = InCheckBoxState == ECheckBoxState::Checked; })
+							.IsChecked_Lambda([this]() { return bDisableActionLimit ? ECheckBoxState::Checked : ECheckBoxState::Unchecked; })
+							.OnCheckStateChanged_Lambda([this](ECheckBoxState InCheckBoxState) 
+								{ 
+									bDisableActionLimit = InCheckBoxState == ECheckBoxState::Checked;
+									GConfig->SetBool(SectionName, DisableActionLimitKey, bDisableActionLimit, GEngineIni);
+									GConfig->Flush(false, GEngineIni);
+								})
 							[
 								SNew(STextBlock)
 								.Text(LOCTEXT("DisableLimit", "Disable action limit for this session"))
@@ -194,6 +205,8 @@ public:
 			MainThreadTasks.Empty();
 		}
 
+		// NOTE: In normal operation, this is never reached.  The window's JOB system will terminate the process.
+
 		// Make sure the window is hidden, because it might take a while for the background thread to finish.
 		Window->HideWindow();
 
@@ -289,13 +302,11 @@ private:
 			Arguments += FString::Printf(TEXT("-Target=\"%s\" "), *Target.Replace(TEXT("\""), TEXT("\"\"")));
 		}
 		Arguments += FString::Printf(TEXT("-LiveCoding -LiveCodingModules=\"%s\" -LiveCodingManifest=\"%s\" -WaitMutex"), *ModulesFileName, *ManifestFileName);
-		if (!bDisableLimit && CompileReason == ELiveCodingCompileReason::Initial)
+		if (!bDisableActionLimit && CompileReason == ELiveCodingCompileReason::Initial)
 		{
-			bool bDisableActionLimit = false;
-			GConfig->GetBool(TEXT("LiveCoding"), TEXT("bDisableActionLimit"), bDisableActionLimit, GEngineIni);
 			int ActionLimit = 1;
-			GConfig->GetInt(TEXT("LiveCoding"), TEXT("ActionLimit"), ActionLimit, GEngineIni);
-			if (!bDisableActionLimit && ActionLimit > 0)
+			GConfig->GetInt(SectionName, ActionLimitKey, ActionLimit, GEngineIni);
+			if (ActionLimit > 0)
 			{
 				Arguments += FString::Printf(TEXT(" -LiveCodingLimit=%d"), ActionLimit);
 			}
