Fixed the description for the action limit message.
Changed the default to 99 when the ini file can't be read.

#rnx

diff --git a/Engine/Source/Programs/LiveCodingConsole/Private/LiveCodingConsole.cpp b/Engine/Source/Programs/LiveCodingConsole/Private/LiveCodingConsole.cpp
--- a/Engine/Source/Programs/LiveCodingConsole/Private/LiveCodingConsole.cpp
+++ b/Engine/Source/Programs/LiveCodingConsole/Private/LiveCodingConsole.cpp
@@ -304,7 +304,7 @@ private:
 		Arguments += FString::Printf(TEXT("-LiveCoding -LiveCodingModules=\"%s\" -LiveCodingManifest=\"%s\" -WaitMutex"), *ModulesFileName, *ManifestFileName);
 		if (!bDisableActionLimit && CompileReason == ELiveCodingCompileReason::Initial)
 		{
-			int ActionLimit = 1;
+			int ActionLimit = 99;
 			GConfig->GetInt(SectionName, ActionLimitKey, ActionLimit, GEngineIni);
 			if (ActionLimit > 0)
 			{
@@ -346,7 +346,7 @@ private:
 			if (ReturnCode == ECompilationResult::LiveCodingLimitError)
 			{
 				const FText Message = LOCTEXT("LimitText", "Live Coding action limit reached.  Do you wish to compile anyway?\n\n"
-				"The limit can be permanently changed or disabled by setting the ActionLimit or DisableActionLimit setting in the engine.");
+				"The limit can be permanently changed or disabled by setting the ActionLimit or DisableActionLimit setting for the LiveCodingConsole program.");
 				const FText Title = LOCTEXT("LimitTitle", "Live Coding Action Limit Reached");
 				EAppReturnType::Type ReturnType = FMessageDialog::Open(EAppMsgType::YesNo, Message, &Title);
 				CompileResult = ReturnType == EAppReturnType::Yes ? ELiveCodingCompileResult::Retry : ELiveCodingCompileResult::Canceled;
