PR #9001: Fix building LiveCodingConsole with clang-cl (Contributed by [REDACTED_PERSON])

#rnx

diff --git a/Engine/Source/Developer/Windows/LiveCodingServer/Private/External/LC_DiaVariant.cpp b/Engine/Source/Developer/Windows/LiveCodingServer/Private/External/LC_DiaVariant.cpp
--- a/Engine/Source/Developer/Windows/LiveCodingServer/Private/External/LC_DiaVariant.cpp
+++ b/Engine/Source/Developer/Windows/LiveCodingServer/Private/External/LC_DiaVariant.cpp
@@ -9,9 +9,13 @@
 // END EPIC MOD
 
 dia::Variant::Variant(IDiaSymbol* symbol)
-	: m_var{ VT_EMPTY }
-	, m_str(nullptr)
+// BEGIN EPIC MOD
+	: m_str(nullptr)
+// END EPIC MOD
 {
+	// BEGIN EPIC MOD
+	m_var.vt = VT_EMPTY;
+	// END EPIC MOD
 	if (symbol->get_value(&m_var) == S_OK)
 	{
 		// the information we're interested in is always stored as string
diff --git a/Engine/Source/Programs/LiveCodingConsole/Private/LiveCodingConsole.cpp b/Engine/Source/Programs/LiveCodingConsole/Private/LiveCodingConsole.cpp
--- a/Engine/Source/Programs/LiveCodingConsole/Private/LiveCodingConsole.cpp
+++ b/Engine/Source/Programs/LiveCodingConsole/Private/LiveCodingConsole.cpp
@@ -36,9 +36,9 @@ static void OnRequestExit()
 class FLiveCodingConsoleApp
 {
 private:
-	static constexpr TCHAR* SectionName = TEXT("LiveCodingConsole");
-	static constexpr TCHAR* DisableActionLimitKey = TEXT("bDisableActionLimit");
-	static constexpr TCHAR* ActionLimitKey = TEXT("ActionLimit");
+	static constexpr const TCHAR* SectionName = TEXT("LiveCodingConsole");
+	static constexpr const TCHAR* DisableActionLimitKey = TEXT("bDisableActionLimit");
+	static constexpr const TCHAR* ActionLimitKey = TEXT("ActionLimit");
 
 	FCriticalSection CriticalSection;
 	FSlateApplication& Slate;
diff --git a/Engine/Source/Programs/UnrealBuildTool/Platform/Windows/UEBuildWindows.cs b/Engine/Source/Programs/UnrealBuildTool/Platform/Windows/UEBuildWindows.cs
--- a/Engine/Source/Programs/UnrealBuildTool/Platform/Windows/UEBuildWindows.cs
+++ b/Engine/Source/Programs/UnrealBuildTool/Platform/Windows/UEBuildWindows.cs
@@ -1909,7 +1909,7 @@ namespace UnrealBuildTool
 					}
 				}
 
-				List<DirectoryReference> VisualStudioDirs = FindVisualStudioInstallations(Compiler).ConvertAll(x => x.BaseDir);
+				List<DirectoryReference> VisualStudioDirs = VisualStudioInstallations.ConvertAll(x => x.BaseDir);
 				foreach (DirectoryReference VisualStudioDir in VisualStudioDirs)
 				{
 					DirectoryReference DiaSdkDir = DirectoryReference.Combine(VisualStudioDir, "DIA SDK");
