Fixed now invalid include path.

#rnx

diff --git a/Engine/Source/Programs/LiveCodingConsole/Private/LiveCodingConsole.cpp b/Engine/Source/Programs/LiveCodingConsole/Private/LiveCodingConsole.cpp
--- a/Engine/Source/Programs/LiveCodingConsole/Private/LiveCodingConsole.cpp
+++ b/Engine/Source/Programs/LiveCodingConsole/Private/LiveCodingConsole.cpp
@@ -16,7 +16,7 @@
 #include "Windows/WindowsHWrapper.h"
 #include "Misc/MonitoredProcess.h"
 #include "LiveCodingManifest.h"
-#include "SourceCodeAccess/Public/ISourceCodeAccessModule.h"
+#include "ISourceCodeAccessModule.h"
 #include "Modules/ModuleInterface.h"
 #include "Widgets/Input/SButton.h"
 #include "Widgets/Input/SCheckBox.h"
