Adding lazy load modules to the manifest.  Unknown modules will generate error messages but not halt patch generation.  This resolves an issue where changes to unloaded modules would prevent any type of patching.

[RN] matt.peters
#rnx

diff --git a/Engine/Source/Developer/Windows/LiveCodingServer/Private/External/LC_LiveProcess.h b/Engine/Source/Developer/Windows/LiveCodingServer/Private/External/LC_LiveProcess.h
--- a/Engine/Source/Developer/Windows/LiveCodingServer/Private/External/LC_LiveProcess.h
+++ b/Engine/Source/Developer/Windows/LiveCodingServer/Private/External/LC_LiveProcess.h
@@ -91,9 +91,16 @@ public:
 	// END EPIC MOD
 
 	// BEGIN EPIC MOD - Allow lazy-loading modules
+	struct LazyLoadedModule
+	{
+		Windows::HMODULE m_moduleBase;
+		bool m_loaded;
+	};
+
 	void AddLazyLoadedModule(const std::wstring moduleName, Windows::HMODULE moduleBase);
 	void SetLazyLoadedModuleAsLoaded(const std::wstring moduleName);
 	bool IsPendingLazyLoadedModule(const std::wstring& moduleName) const;
+	const types::unordered_map<std::wstring, LazyLoadedModule>& GetLazyLoadedModules() const { return m_lazyLoadedModules; }
 	Windows::HMODULE GetLazyLoadedModuleBase(const std::wstring& moduleName) const;
 	// END EPIC MOD
 
@@ -143,12 +150,6 @@ private:
 	// END EPIC MOD
 
 	// BEGIN EPIC MOD - Allow lazy-loading modules
-	struct LazyLoadedModule
-	{
-		Windows::HMODULE m_moduleBase;
-		bool m_loaded;
-	};
-
 	types::unordered_map<std::wstring, LazyLoadedModule> m_lazyLoadedModules;
 	// END EPIC MOD
 
diff --git a/Engine/Source/Developer/Windows/LiveCodingServer/Private/External/LC_ServerCommandThread.cpp b/Engine/Source/Developer/Windows/LiveCodingServer/Private/External/LC_ServerCommandThread.cpp
--- a/Engine/Source/Developer/Windows/LiveCodingServer/Private/External/LC_ServerCommandThread.cpp
+++ b/Engine/Source/Developer/Windows/LiveCodingServer/Private/External/LC_ServerCommandThread.cpp
@@ -815,9 +815,23 @@ void ServerCommandThread::CompileChanges(bool didAllProcessesMakeProgress, comma
 				ValidModules.Add(liveModule->GetModuleName().c_str());
 			}
 
+			// Build a list of loaded modules which are not enabled
+			TSet<FString> LazyLoadModules;
+			for (LiveProcess* liveProcess : m_liveProcesses)
+			{
+				const std::unordered_map<std::wstring, LiveProcess::LazyLoadedModule>& lazyLoadedModules = liveProcess->GetLazyLoadedModules();
+				for (const std::unordered_map<std::wstring, LiveProcess::LazyLoadedModule>::value_type& kvp : lazyLoadedModules)
+				{
+					if (!kvp.second.m_loaded)
+					{
+						LazyLoadModules.Add(FString(kvp.first.c_str()));
+					}
+				}
+			}
+
 			// Execute the compile
 			TArray<FString> RequiredModules;
-			ELiveCodingCompileResult CompileResult = CompileDelegate.Execute(Targets, ValidModules, RequiredModules, ModuleToModuleFiles, CompileReason);
+			ELiveCodingCompileResult CompileResult = CompileDelegate.Execute(Targets, ValidModules, LazyLoadModules, RequiredModules, ModuleToModuleFiles, CompileReason);
 			if (CompileResult == ELiveCodingCompileResult::Success)
 			{
 				break;
@@ -877,13 +891,8 @@ void ServerCommandThread::CompileChanges(bool didAllProcessesMakeProgress, comma
 
 				if (!bFoundNameMatch)
 				{
-					LC_ERROR_USER("Live coding is not enabled for %S.", ModuleFileName.c_str());
-					LC_ERROR_USER("Configure the list of enabled modules from the Live Coding section of the editor preferences window.");
-					// BEGIN EPIC MOD
-					postCompileResult = commands::PostCompileResult::Failure;
-					// END EPIC MOD
-					GLiveCodingServer->GetCompileFinishedDelegate().ExecuteIfBound(ELiveCodingResult::Error, *FString::Printf(TEXT("Live coding not enabled for %s"), ModuleFileName.c_str()));
-					return;
+					LC_WARNING_USER("The module '%S' has not been loaded by any process. Changes will be ignored.", ModuleFileName.c_str());
+					continue;
 				}
 			}
 
diff --git a/Engine/Source/Developer/Windows/LiveCodingServer/Public/ILiveCodingServer.h b/Engine/Source/Developer/Windows/LiveCodingServer/Public/ILiveCodingServer.h
--- a/Engine/Source/Developer/Windows/LiveCodingServer/Public/ILiveCodingServer.h
+++ b/Engine/Source/Developer/Windows/LiveCodingServer/Public/ILiveCodingServer.h
@@ -75,7 +75,7 @@ public:
 	DECLARE_DELEGATE_TwoParams(FLogOutputDelegate, ELiveCodingLogVerbosity, const wchar_t*);
 	virtual FLogOutputDelegate& GetLogOutputDelegate() = 0;
 
-	DECLARE_DELEGATE_RetVal_FiveParams(ELiveCodingCompileResult, FCompileDelegate, const TArray<FString>&, const TArray<FString>&, TArray<FString>&, FModuleToModuleFiles&, ELiveCodingCompileReason);
+	DECLARE_DELEGATE_RetVal_SixParams(ELiveCodingCompileResult, FCompileDelegate, const TArray<FString>&, const TArray<FString>&, const TSet<FString>&, TArray<FString>&, FModuleToModuleFiles&, ELiveCodingCompileReason);
 	virtual FCompileDelegate& GetCompileDelegate() = 0;
 
 	DECLARE_DELEGATE(FCompileStartedDelegate);
diff --git a/Engine/Source/Programs/LiveCodingConsole/Private/LiveCodingConsole.cpp b/Engine/Source/Programs/LiveCodingConsole/Private/LiveCodingConsole.cpp
--- a/Engine/Source/Programs/LiveCodingConsole/Private/LiveCodingConsole.cpp
+++ b/Engine/Source/Programs/LiveCodingConsole/Private/LiveCodingConsole.cpp
@@ -265,7 +265,8 @@ private:
 		LogWidget->AppendLine(GetLogColor(Verbosity), MoveTemp(Text));
 	}
 
-	ELiveCodingCompileResult CompilePatch(const TArray<FString>& Targets, const TArray<FString>& ValidModules, TArray<FString>& RequiredModules, FModuleToModuleFiles& ModuleToModuleFiles, ELiveCodingCompileReason CompileReason)
+	ELiveCodingCompileResult CompilePatch(const TArray<FString>& Targets, const TArray<FString>& ValidModules, const TSet<FString>& LazyLoadModules,
+		TArray<FString>& RequiredModules, FModuleToModuleFiles& ModuleToModuleFiles, ELiveCodingCompileReason CompileReason)
 	{
 		// Update the compile start time. This gets copied into the last patch time once a patch has been confirmed to have been applied.
 		NextPatchStartTime = FDateTime::UtcNow();
@@ -284,8 +285,36 @@ private:
 		FPaths::MakePlatformFilename(Executable);
 
 		// Write out the list of lazy-loaded modules for UBT to check
-		FString ModulesFileName = FPaths::ConvertRelativePathToFull(FPaths::EngineIntermediateDir() / TEXT("LiveCodingModules.txt"));
-		FFileHelper::SaveStringArrayToFile(ValidModules, *ModulesFileName);
+		FString ModulesFileName = FPaths::ConvertRelativePathToFull(FPaths::EngineIntermediateDir() / TEXT("LiveCodingModules.json"));
+
+		// Create and open a new Json file for writing, and initialize a JsonWriter to serialize the contents
+		{
+			TArray<uint8> TCharJsonData;
+			FMemoryWriter MemWriter(TCharJsonData, true);
+			{
+				TSharedRef<TJsonWriter<>> JsonWriter = TJsonWriterFactory<>::Create(&MemWriter);
+				JsonWriter->WriteObjectStart();
+				JsonWriter->WriteArrayStart(TEXT("EnabledModules"));
+				for (const FString& name : ValidModules)
+				{
+					JsonWriter->WriteValue(name);
+				}
+				JsonWriter->WriteArrayEnd();
+				JsonWriter->WriteArrayStart(TEXT("LazyLoadModules"));
+				for (const FString& name : LazyLoadModules)
+				{
+					JsonWriter->WriteValue(name);
+				}
+				JsonWriter->WriteArrayEnd();
+				JsonWriter->WriteObjectEnd();
+			}
+			MemWriter.Close();
+
+			TCharJsonData.AddZeroed(sizeof(TCHAR));
+			auto Utf8Data = StringCast<UTF8CHAR>((const TCHAR*)TCharJsonData.GetData());
+			TSharedPtr<FArchive> Ar(IFileManager::Get().CreateFileWriter(*ModulesFileName));
+			Ar->Serialize(const_cast<UTF8CHAR*>(Utf8Data.Get()), Utf8Data.Length());
+		}
 
 		// Delete the output file for non-allowed modules
 		FString ModulesOutputFileName = ModulesFileName + TEXT(".out");
