Improve the out-of-date check for object files and the dlls.  This is now required since object files can now be ignored.  Once the dll for an ignored object file is loaded by the engine, any future live coding requests will link the previously ignored object files.

#rnx

diff --git a/Engine/Source/Developer/Windows/LiveCodingServer/Private/External/LC_LiveModule.cpp b/Engine/Source/Developer/Windows/LiveCodingServer/Private/External/LC_LiveModule.cpp
--- a/Engine/Source/Developer/Windows/LiveCodingServer/Private/External/LC_LiveModule.cpp
+++ b/Engine/Source/Developer/Windows/LiveCodingServer/Private/External/LC_LiveModule.cpp
@@ -1210,7 +1210,9 @@ void LiveModule::Load(symbols::Provider* provider, symbols::DiaCompilandDB* diaC
 	}
 
 	// now that all the databases are built, store their info into the module cache
-	m_mainModuleToken = m_moduleCache->Insert(m_symbolDB, m_contributionDB, m_compilandDB, m_thunkDB, m_imageSectionDB);
+	// BEGIN EPIC MOD
+	m_mainModuleToken = m_moduleCache->Insert(m_symbolDB, m_contributionDB, m_compilandDB, m_thunkDB, m_imageSectionDB, provider->lastModificationTime);
+	// END EPIC MOD
 }
 
 
@@ -3674,6 +3676,9 @@ LiveModule::ErrorType::Enum LiveModule::Update(FileAttributeCache* fileCache, Di
 	executable::DestroyImageSectionDB(patchImageSections);
 	executable::CloseImage(patchImage);
 
+	// BEGIN EPIC MOD
+	uint64_t patchLastModicationTime = patchSymbolProvider->lastModificationTime;
+	// END EPIC MOD
 	symbols::Close(patchSymbolProvider);
 
 
@@ -3684,7 +3689,9 @@ LiveModule::ErrorType::Enum LiveModule::Update(FileAttributeCache* fileCache, Di
 
 	// store the new databases into the module cache
 	ModulePatch* compiledModulePatch = nullptr;
-	const size_t token = m_moduleCache->Insert(patch_symbolDB, patch_contributionDB, patch_compilandDB, patch_thunkDB, patch_imageSectionDB);
+	// BEGIN EPIC MOD
+	const size_t token = m_moduleCache->Insert(patch_symbolDB, patch_contributionDB, patch_compilandDB, patch_thunkDB, patch_imageSectionDB, patchLastModicationTime);
+	// END EPIC MOD
 	{
 		for (size_t p = 0u; p < processCount; ++p)
 		{
@@ -4895,3 +4902,29 @@ void LiveModule::OnCompiledFile(const symbols::ObjPath& objPath, symbols::Compil
 		LC_ERROR_USER("Failed to compile %s (%.3fs) (Exit code: 0x%X)", objPath.c_str(), compileTime, compileResult.exitCode);
 	}
 }
+
+// BEGIN EPIC MOD
+bool LiveModule::IsModifiedSource(const wchar_t* sourceFile) const
+{
+	if (m_moduleCache == nullptr || m_moduleCache->GetSize() == 0)
+	{
+		return false;
+	}
+
+	Filesystem::PathAttributes attr = Filesystem::GetAttributes(sourceFile);
+	if (!Filesystem::DoesExist(attr))
+	{
+		return false;
+	}
+
+	for (size_t i = 0; i < m_moduleCache->GetSize(); ++i)
+	{
+		const ModuleCache::Data& data = m_moduleCache->GetEntry(i);
+		if (data.lastModificationTime > attr.lastModificationTime)
+		{
+			return false;
+		}
+	}
+	return true;
+}
+// END EPIC MOD
\ No newline at end of file
diff --git a/Engine/Source/Developer/Windows/LiveCodingServer/Private/External/LC_LiveModule.h b/Engine/Source/Developer/Windows/LiveCodingServer/Private/External/LC_LiveModule.h
--- a/Engine/Source/Developer/Windows/LiveCodingServer/Private/External/LC_LiveModule.h
+++ b/Engine/Source/Developer/Windows/LiveCodingServer/Private/External/LC_LiveModule.h
@@ -92,6 +92,10 @@ public:
 
 	bool HasInstalledPatches(void) const;
 
+	// BEGIN EPIC MOD
+	bool IsModifiedSource(const wchar_t* sourceFile) const;
+	// END EPIC MOD
+
 private:
 	void UpdateDirectoryCache(const ImmutableString& path, symbols::Dependency* dependency, DirectoryCache* cache);
 
diff --git a/Engine/Source/Developer/Windows/LiveCodingServer/Private/External/LC_ModuleCache.cpp b/Engine/Source/Developer/Windows/LiveCodingServer/Private/External/LC_ModuleCache.cpp
--- a/Engine/Source/Developer/Windows/LiveCodingServer/Private/External/LC_ModuleCache.cpp
+++ b/Engine/Source/Developer/Windows/LiveCodingServer/Private/External/LC_ModuleCache.cpp
@@ -14,13 +14,14 @@ ModuleCache::ModuleCache(void)
 	m_cache.reserve(128u);
 }
 
-
-size_t ModuleCache::Insert(const symbols::SymbolDB* symbolDb, const symbols::ContributionDB* contributionDb, const symbols::CompilandDB* compilandDb, const symbols::ThunkDB* thunkDb, const symbols::ImageSectionDB* imageSectionDb)
+// BEGIN EPIC MOD
+size_t ModuleCache::Insert(const symbols::SymbolDB* symbolDb, const symbols::ContributionDB* contributionDb, const symbols::CompilandDB* compilandDb, const symbols::ThunkDB* thunkDb, const symbols::ImageSectionDB* imageSectionDb, uint64_t lastModicationTime)
+// END EPIC MOD
 {
 	CriticalSection::ScopedLock lock(&m_cs);
 
 	const size_t token = m_cache.size();
-	m_cache.push_back(Data { static_cast<uint16_t>(token), symbolDb, contributionDb, compilandDb, thunkDb, imageSectionDb, {} });
+	m_cache.push_back(Data { static_cast<uint16_t>(token), symbolDb, contributionDb, compilandDb, thunkDb, imageSectionDb, lastModicationTime, {} });
 
 	return token;
 }
diff --git a/Engine/Source/Developer/Windows/LiveCodingServer/Private/External/LC_ModuleCache.h b/Engine/Source/Developer/Windows/LiveCodingServer/Private/External/LC_ModuleCache.h
--- a/Engine/Source/Developer/Windows/LiveCodingServer/Private/External/LC_ModuleCache.h
+++ b/Engine/Source/Developer/Windows/LiveCodingServer/Private/External/LC_ModuleCache.h
@@ -39,6 +39,9 @@ public:
 		const symbols::CompilandDB* compilandDb;
 		const symbols::ThunkDB* thunkDb;
 		const symbols::ImageSectionDB* imageSectionDb;
+		// BEGIN EPIC MOD
+		uint64_t lastModificationTime;
+		// END EPIC MOD
 
 		types::vector<ProcessData> processes;	// all processes that this patch is loaded into
 	};
@@ -61,7 +64,9 @@ public:
 	
 	// adds an entry to the cache. does not take ownership of the databases.
 	// returns a token for registering a process associated with this entry.
-	size_t Insert(const symbols::SymbolDB* symbolDb, const symbols::ContributionDB* contributionDb, const symbols::CompilandDB* compilandDb, const symbols::ThunkDB* thunkDb, const symbols::ImageSectionDB* imageSectionDb);
+	// BEGIN EPIC MOD
+	size_t Insert(const symbols::SymbolDB* symbolDb, const symbols::ContributionDB* contributionDb, const symbols::CompilandDB* compilandDb, const symbols::ThunkDB* thunkDb, const symbols::ImageSectionDB* imageSectionDb, uint64_t lastModicationTime);
+	// END EPIC MOD
 
 	// associates a process with an entry identified by a previously returned token.
 	void RegisterProcess(size_t token, LiveProcess* liveProcess, void* moduleBase);
diff --git a/Engine/Source/Developer/Windows/LiveCodingServer/Private/External/LC_ServerCommandThread.cpp b/Engine/Source/Developer/Windows/LiveCodingServer/Private/External/LC_ServerCommandThread.cpp
--- a/Engine/Source/Developer/Windows/LiveCodingServer/Private/External/LC_ServerCommandThread.cpp
+++ b/Engine/Source/Developer/Windows/LiveCodingServer/Private/External/LC_ServerCommandThread.cpp
@@ -864,43 +864,53 @@ void ServerCommandThread::CompileChanges(bool didAllProcessesMakeProgress, comma
 		symbols::ResetCachedUnityManifests();
 
 		// Build up a list of all the modified object files in each module
-		types::unordered_set<std::wstring> ValidModuleFileNames;
+		types::unordered_map<std::wstring, const LiveModule*> EnabledModulesByName;
 		for (const LiveModule* liveModule : m_liveModules)
 		{
-			ValidModuleFileNames.insert(liveModule->GetModuleName());
+			EnabledModulesByName[liveModule->GetModuleName()] = liveModule;
 		}
 
 		for(const TPair<FString, FModuleFiles>& Pair : ModuleToModuleFiles)
 		{
+			const LiveModule* liveModule = nullptr;
 			std::wstring ModuleFileName = Filesystem::NormalizePath(*Pair.Key).GetString();
-			if(ValidModuleFileNames.find(ModuleFileName) == ValidModuleFileNames.end())
+			types::unordered_map<std::wstring, const LiveModule*>::iterator ix = EnabledModulesByName.find(ModuleFileName);
+			if (ix == EnabledModulesByName.end())
 			{
 				// We couldn't find this exact module filename, but this could be a staged executable. See if we can just match the name.
 				std::wstring ModuleFileNameOnly = Filesystem::GetFilename(ModuleFileName.c_str()).GetString();
 
-				bool bFoundNameMatch = false;
-				for (const LiveModule* liveModule : m_liveModules)
+				for (const LiveModule* testLiveModule : m_liveModules)
 				{
-					if (ModuleFileNameOnly == Filesystem::GetFilename(liveModule->GetModuleName().c_str()).GetString())
+					if (ModuleFileNameOnly == Filesystem::GetFilename(testLiveModule->GetModuleName().c_str()).GetString())
 					{
-						ModuleFileName = liveModule->GetModuleName();
-						bFoundNameMatch = true;
+						ModuleFileName = testLiveModule->GetModuleName();
+						liveModule = testLiveModule;
 						break;
 					}
 				}
 
-				if (!bFoundNameMatch)
+				if (liveModule == nullptr)
 				{
 					LC_WARNING_USER("The module '%S' has not been loaded by any process. Changes will be ignored.", ModuleFileName.c_str());
 					continue;
 				}
 			}
+			else
+			{
+				liveModule = ix->second;
+			}
 
 			types::vector<symbols::ModifiedObjFile> ObjectFiles;
 			for (const FString& ObjectFile : Pair.Value.Objects)
 			{
 				std::wstring NormalizedObjectFile = Filesystem::NormalizePath(*ObjectFile).GetString();
 
+				if (!liveModule->IsModifiedSource(NormalizedObjectFile.c_str()))
+				{
+					continue;
+				}
+
 				// If this file has a .lc.obj suffix, temporarily replace the original .obj file while generating the patch.
 				// It'd be nice to track this explicitly inside Live++ and just load the new file, but it requires a lot of changes and would make upgrades difficult.
 				static const TCHAR Suffix[] = TEXT(".lc.obj");
diff --git a/Engine/Source/Programs/LiveCodingConsole/Private/LiveCodingConsole.cpp b/Engine/Source/Programs/LiveCodingConsole/Private/LiveCodingConsole.cpp
--- a/Engine/Source/Programs/LiveCodingConsole/Private/LiveCodingConsole.cpp
+++ b/Engine/Source/Programs/LiveCodingConsole/Private/LiveCodingConsole.cpp
@@ -51,8 +51,6 @@ private:
 	bool bDisableActionLimit;
 	bool bHasReinstancingProcess;
 	bool bWarnOnRestart;
-	FDateTime LastPatchTime;
-	FDateTime NextPatchStartTime;
 
 public:
 	FLiveCodingConsoleApp(FSlateApplication& InSlate, ILiveCodingServer& InServer)
@@ -62,8 +60,6 @@ public:
 		, bDisableActionLimit(false)
 		, bHasReinstancingProcess(false)
 		, bWarnOnRestart(false)
-		, LastPatchTime(FDateTime::MinValue())
-		, NextPatchStartTime(FDateTime::MinValue())
 	{
 	}
 
@@ -268,9 +264,6 @@ private:
 	ELiveCodingCompileResult CompilePatch(const TArray<FString>& Targets, const TArray<FString>& ValidModules, const TSet<FString>& LazyLoadModules,
 		TArray<FString>& RequiredModules, FModuleToModuleFiles& ModuleToModuleFiles, ELiveCodingCompileReason CompileReason)
 	{
-		// Update the compile start time. This gets copied into the last patch time once a patch has been confirmed to have been applied.
-		NextPatchStartTime = FDateTime::UtcNow();
-
 		// Get the UBT path
 		FString Executable;
 		FString Entry;
@@ -400,22 +393,12 @@ private:
 		// Override the linker path
 		Server.SetLinkerPath(*Manifest.LinkerPath, Manifest.LinkerEnvironment);
 
-		// Strip out all the files that haven't been modified
-		IFileManager& FileManager = IFileManager::Get();
+		// Add all the sources regardless of modification time.  The patch will exclude old ones
 		for(TPair<FString, TArray<FString>>& Pair : Manifest.BinaryToObjectFiles)
 		{
-			FDateTime MinTimeStamp = FileManager.GetTimeStamp(*Pair.Key);
-			if(LastPatchTime > MinTimeStamp)
-			{
-				MinTimeStamp = LastPatchTime;
-			}
-
 			for (const FString& ObjectFileName : Pair.Value)
 			{
-				if (FileManager.GetTimeStamp(*ObjectFileName) > MinTimeStamp)
-				{
-					ModuleToModuleFiles.FindOrAdd(Pair.Key).Objects.Add(ObjectFileName);
-				}
+				ModuleToModuleFiles.FindOrAdd(Pair.Key).Objects.Add(ObjectFileName);
 			}
 		}
 
@@ -428,10 +411,6 @@ private:
 			}
 		}
 
-		// Force the next patch time to include any CPP files modified this time around.  This avoids the issue 
-		// where a patch might be generated for a file twice instead of just once.
-		NextPatchStartTime = FDateTime::UtcNow();
-
 		return ELiveCodingCompileResult::Success;
 	}
 
@@ -498,7 +477,6 @@ private:
 		FScopeLock Lock(&CriticalSection);
 		MainThreadTasks.Add(FSimpleDelegate::CreateRaw(this, &FLiveCodingConsoleApp::OnCompileStarted));
 		bRequestCancel = false;
-		NextPatchStartTime = FDateTime::UtcNow();
 	}
 
 	void OnCompileStarted()
@@ -524,11 +502,6 @@ private:
 	{
 		FScopeLock Lock(&CriticalSection);
 		MainThreadTasks.Add(FSimpleDelegate::CreateLambda([this, Result, Status](){ OnCompileFinished(Result, Status); }));
-
-		if (Result == ELiveCodingResult::Success)
-		{
-			LastPatchTime = NextPatchStartTime;
-		}
 	}
 
 	void OnCompileFinished(ELiveCodingResult Result, const FString& Status)
