Use the FMessageDialog overloads that pass the optional title by-value

#rnx

diff --git a/Engine/Plugins/Bridge/Source/Bridge/Private/UI/BridgeUIManager.cpp b/Engine/Plugins/Bridge/Source/Bridge/Private/UI/BridgeUIManager.cpp
--- a/Engine/Plugins/Bridge/Source/Bridge/Private/UI/BridgeUIManager.cpp
+++ b/Engine/Plugins/Bridge/Source/Bridge/Private/UI/BridgeUIManager.cpp
@@ -167,8 +167,7 @@ void FBridgeUIManagerImpl::CreateWindow()
 // 	{
 // 		if (!WebBrowserPlugin->IsEnabled())
 // 		{
-// 			const FText Title = FText::FromString(TEXT("Enable Web Browser Plugin"));
-// 			FMessageDialog::Open(EAppMsgType::Ok, FText::FromString(TEXT("Quixel Bridge requires the “Web Browser” plugin, which is disabled. Go to Edit > Plugins and search for “Web Browser” to enable it.")), &Title);
+// 			FMessageDialog::Open(EAppMsgType::Ok, FText::FromString(TEXT("Quixel Bridge requires the “Web Browser” plugin, which is disabled. Go to Edit > Plugins and search for “Web Browser” to enable it.")), FText::FromString(TEXT("Enable Web Browser Plugin")));
 // 			return;
 // 		}
 // 	}
diff --git a/Engine/Plugins/Editor/PluginBrowser/Templates/EditorMode/Source/PLUGIN_NAME/Private/Tools/PLUGIN_NAMESimpleTool.cpp b/Engine/Plugins/Editor/PluginBrowser/Templates/EditorMode/Source/PLUGIN_NAME/Private/Tools/PLUGIN_NAMESimpleTool.cpp
--- a/Engine/Plugins/Editor/PluginBrowser/Templates/EditorMode/Source/PLUGIN_NAME/Private/Tools/PLUGIN_NAMESimpleTool.cpp
+++ b/Engine/Plugins/Editor/PluginBrowser/Templates/EditorMode/Source/PLUGIN_NAME/Private/Tools/PLUGIN_NAMESimpleTool.cpp
@@ -86,7 +86,7 @@ void UPLUGIN_NAMESimpleTool::OnClicked(const FInputDeviceRay& ClickPos)
 
 			FText Title = LOCTEXT("ActorInfoDialogTitle", "Actor Info");
 			// JAH TODO: consider if we can highlight the actor prior to opening the dialog box or make it non-modal
-			FMessageDialog::Open(EAppMsgType::Ok, ActorInfoMsg, &Title);
+			FMessageDialog::Open(EAppMsgType::Ok, ActorInfoMsg, Title);
 		}
 	}
 }
diff --git a/Engine/Source/Programs/LiveCodingConsole/Private/LiveCodingConsole.cpp b/Engine/Source/Programs/LiveCodingConsole/Private/LiveCodingConsole.cpp
--- a/Engine/Source/Programs/LiveCodingConsole/Private/LiveCodingConsole.cpp
+++ b/Engine/Source/Programs/LiveCodingConsole/Private/LiveCodingConsole.cpp
@@ -370,7 +370,7 @@ private:
 				const FText Message = LOCTEXT("LimitText", "Live Coding action limit reached.  Do you wish to compile anyway?\n\n"
 				"The limit can be permanently changed or disabled by setting the ActionLimit or DisableActionLimit setting for the LiveCodingConsole program.");
 				const FText Title = LOCTEXT("LimitTitle", "Live Coding Action Limit Reached");
-				EAppReturnType::Type ReturnType = FMessageDialog::Open(EAppMsgType::YesNo, Message, &Title);
+				EAppReturnType::Type ReturnType = FMessageDialog::Open(EAppMsgType::YesNo, Message, Title);
 				CompileResult = ReturnType == EAppReturnType::Yes ? ELiveCodingCompileResult::Retry : ELiveCodingCompileResult::Canceled;
 			}
 
@@ -432,7 +432,7 @@ private:
 		{
 			const FText Message = LOCTEXT("RestartWarningText", "Restarting after patching while re-instancing was enabled can lead to unexpected results.\r\n\r\nDo you wish to continue?");
 			const FText Title = LOCTEXT("RestartWarningTitle", "Restarting after patching while re-instancing was enabled?");
-			EAppReturnType::Type ReturnType = FMessageDialog::Open(EAppMsgType::YesNo, Message, &Title);
+			EAppReturnType::Type ReturnType = FMessageDialog::Open(EAppMsgType::YesNo, Message, Title);
 			if (ReturnType != EAppReturnType::Yes)
 			{
 				return FReply::Handled();
diff --git a/Engine/Source/Runtime/Apple/MetalRHI/Private/MetalRHI.cpp b/Engine/Source/Runtime/Apple/MetalRHI/Private/MetalRHI.cpp
--- a/Engine/Source/Runtime/Apple/MetalRHI/Private/MetalRHI.cpp
+++ b/Engine/Source/Runtime/Apple/MetalRHI/Private/MetalRHI.cpp
@@ -75,7 +75,7 @@ static void ValidateTargetedRHIFeatureLevelExists(EShaderPlatform Platform)
 		FText LocalizedMsg = FText::Format(NSLOCTEXT("MetalRHI", "ShaderPlatformUnavailable","Shader platform: {ShaderPlatform} was not cooked! Please enable this shader platform in the project's target settings."),Args);
 		
 		FText Title = NSLOCTEXT("MetalRHI", "ShaderPlatformUnavailableTitle","Shader Platform Unavailable");
-		FMessageDialog::Open(EAppMsgType::Ok, LocalizedMsg, &Title);
+		FMessageDialog::Open(EAppMsgType::Ok, LocalizedMsg, Title);
 		FPlatformMisc::RequestExit(true);
 		
 		METAL_FATAL_ERROR(TEXT("Shader platform: %s was not cooked! Please enable this shader platform in the project's target settings."), *LegacyShaderPlatformToShaderFormat(Platform).ToString());
