Fix notifications never disappearing in LiveCoding

diff --git a/Engine/Source/Programs/LiveCodingConsole/Private/LiveCodingConsole.cpp b/Engine/Source/Programs/LiveCodingConsole/Private/LiveCodingConsole.cpp
--- a/Engine/Source/Programs/LiveCodingConsole/Private/LiveCodingConsole.cpp
+++ b/Engine/Source/Programs/LiveCodingConsole/Private/LiveCodingConsole.cpp
@@ -187,6 +187,8 @@ public:
 		{
 			BeginExitIfRequested();
 
+			FTSTicker::GetCoreTicker().Tick(FApp::GetDeltaTime());
+
 			Slate.PumpMessages();
 			Slate.Tick();
 
