[LiveCodingConsole]
* Fixed link error

diff --git a/Engine/Source/Programs/LiveCodingConsole/Private/LiveCodingConsole.cpp b/Engine/Source/Programs/LiveCodingConsole/Private/LiveCodingConsole.cpp
--- a/Engine/Source/Programs/LiveCodingConsole/Private/LiveCodingConsole.cpp
+++ b/Engine/Source/Programs/LiveCodingConsole/Private/LiveCodingConsole.cpp
@@ -28,6 +28,12 @@
 
 IMPLEMENT_APPLICATION(LiveCodingConsole, "LiveCodingConsole");
 
+#if IS_MONOLITHIC
+const TCHAR* GLiveCodingEngineDir = nullptr;
+const TCHAR* GLiveCodingProject = nullptr;
+#endif
+
+
 static void OnRequestExit()
 {
 	RequestEngineExit(TEXT("LiveCoding console closed"));
