Added branch name to Live Coding Console window title, to better differentiate instances when working on multiple branches

diff --git a/Engine/Source/Programs/LiveCodingConsole/Private/LiveCodingConsole.cpp b/Engine/Source/Programs/LiveCodingConsole/Private/LiveCodingConsole.cpp
--- a/Engine/Source/Programs/LiveCodingConsole/Private/LiveCodingConsole.cpp
+++ b/Engine/Source/Programs/LiveCodingConsole/Private/LiveCodingConsole.cpp
@@ -23,6 +23,8 @@
 #include "Widgets/Text/STextBlock.h"
 #include "Misc/CompilationResult.h"
 #include "Misc/MessageDialog.h"
+#include "Misc/EngineBuildSettings.h"
+#include "Misc/EngineVersion.h"
 
 #define LOCTEXT_NAMESPACE "LiveCodingConsole"
 
@@ -226,7 +228,8 @@ private:
 		{
 			FFormatNamedArguments Args;
 			Args.Add(TEXT("ProjectName"), FText::FromString(ProjectName));
-			return FText::Format(LOCTEXT("WindowTitleWithProject", "{ProjectName} - Live Coding"), Args);
+			Args.Add(TEXT("Branch"), FEngineBuildSettings::IsPerforceBuild() ? FText::FromString(" - " + FEngineVersion::Current().GetBranch()) : FText::GetEmpty());
+			return FText::Format(LOCTEXT("WindowTitleWithProject", "{ProjectName}{Branch} - Live Coding"), Args);
 		}
 		return LOCTEXT("WindowTitle", "Live Coding");
 	}
