[LiveCoding]
* Fixed so LiveCodingConsole target throw an error if trying to compile it with LPP2 enabled

diff --git a/Engine/Source/Developer/Windows/LiveCoding/LiveCoding.Build.cs b/Engine/Source/Developer/Windows/LiveCoding/LiveCoding.Build.cs
--- a/Engine/Source/Developer/Windows/LiveCoding/LiveCoding.Build.cs
+++ b/Engine/Source/Developer/Windows/LiveCoding/LiveCoding.Build.cs
@@ -5,7 +5,7 @@ using System.IO;
 
 public class LiveCoding : ModuleRules
 {
-	static bool UseLiveCoding2 = false;
+	public const bool UseLiveCoding2 = false;
 
 	public LiveCoding(ReadOnlyTargetRules Target) : base(Target)
 	{
@@ -35,11 +35,13 @@ public class LiveCoding : ModuleRules
             PrivateDefinitions.Add("USE_DEBUG_LIVE_CODING_CONSOLE=0");
         }
 
+		bool UseLiveCoding2ToNotGetCompileWarning = UseLiveCoding2;
+
 		if (!Target.Architecture.bIsX64)
 		{
 			PrivateDefinitions.Add("LC_VERSION=0");
 		}
-		else if (UseLiveCoding2)
+		else if (UseLiveCoding2ToNotGetCompileWarning)
 		{
 			PrivateDefinitions.Add("LC_VERSION=2");
 			PrivateIncludePaths.Add("ThirdParty/LivePlusPlus");
diff --git a/Engine/Source/Programs/LiveCodingConsole/LiveCodingConsole.Target.cs b/Engine/Source/Programs/LiveCodingConsole/LiveCodingConsole.Target.cs
--- a/Engine/Source/Programs/LiveCodingConsole/LiveCodingConsole.Target.cs
+++ b/Engine/Source/Programs/LiveCodingConsole/LiveCodingConsole.Target.cs
@@ -1,13 +1,19 @@
 // Copyright Epic Games, Inc. All Rights Reserved.
 
 using UnrealBuildTool;
-using System.Collections.Generic;
+using Microsoft.Extensions.Logging;
 
 [SupportedPlatforms("Win64")]
 public class LiveCodingConsoleTarget : TargetRules
 {
 	public LiveCodingConsoleTarget(TargetInfo Target) : base(Target)
 	{
+		bool UseLiveCoding2ToNotGetCompileWarning = LiveCoding.UseLiveCoding2;
+		if (UseLiveCoding2ToNotGetCompileWarning)
+		{
+			throw new BuildException("ERROR: LiveCodingConsole is not used when LiveCoding2 is enabled. You should not build this target");
+		}
+
 		Type = TargetType.Program;
 		LinkType = TargetLinkType.Monolithic;
 		IncludeOrderVersion = EngineIncludeOrderVersion.Latest;
