[LiveCodingConsole]
* Fixed bug where it is not possible to generate solution when live coding 2 is enabled

diff --git a/Engine/Source/Programs/LiveCodingConsole/LiveCodingConsole.Target.cs b/Engine/Source/Programs/LiveCodingConsole/LiveCodingConsole.Target.cs
--- a/Engine/Source/Programs/LiveCodingConsole/LiveCodingConsole.Target.cs
+++ b/Engine/Source/Programs/LiveCodingConsole/LiveCodingConsole.Target.cs
@@ -9,7 +9,7 @@ public class LiveCodingConsoleTarget : TargetRules
 	public LiveCodingConsoleTarget(TargetInfo Target) : base(Target)
 	{
 		bool UseLiveCoding2ToNotGetCompileWarning = LiveCoding.UseLiveCoding2;
-		if (UseLiveCoding2ToNotGetCompileWarning)
+		if (UseLiveCoding2ToNotGetCompileWarning && !bGenerateProjectFiles)
 		{
 			throw new BuildException("ERROR: LiveCodingConsole is not used when LiveCoding2 is enabled. You should not build this target");
 		}
